Editgrate 4.25 audio features from project stream:

diff --git a/Engine/Source/Runtime/Engine/Private/SoundBase.cpp b/Engine/Source/Runtime/Engine/Private/SoundBase.cpp
--- a/Engine/Source/Runtime/Engine/Private/SoundBase.cpp
+++ b/Engine/Source/Runtime/Engine/Private/SoundBase.cpp
@@ -135,7 +135,7 @@ USoundClass* USoundBase::GetSoundClass() const
 	return SoundClassObject;
 }
 
-USoundSubmix* USoundBase::GetSoundSubmix() const
+USoundSubmixBase* USoundBase::GetSoundSubmix() const
 {
 	return SoundSubmixObject;
 }
diff --git a/Engine/Source/Runtime/Engine/Private/SoundSubmix.cpp b/Engine/Source/Runtime/Engine/Private/SoundSubmix.cpp
--- a/Engine/Source/Runtime/Engine/Private/SoundSubmix.cpp
+++ b/Engine/Source/Runtime/Engine/Private/SoundSubmix.cpp
@@ -13,13 +13,27 @@
 #include "Widgets/Notifications/SNotificationList.h"
 #endif // WITH_EDITOR
 
-USoundSubmix::USoundSubmix(const FObjectInitializer& ObjectInitializer)
+static int32 ClearBrokenSubmixAssetsCVar = 0;
+FAutoConsoleVariableRef CVarFixUpBrokenSubmixAssets(
+	TEXT("au.submix.clearbrokensubmixassets"),
+	ClearBrokenSubmixAssetsCVar,
+	TEXT("If fixed, will verify that we don't have a submix list a child submix that doesn't have it as it's parent, or vice versa.\n")
+	TEXT("0: Disable, >0: Enable"),
+	ECVF_Default);
+
+USoundSubmixWithParentBase::USoundSubmixWithParentBase(const FObjectInitializer& ObjectInitializer)
 	: Super(ObjectInitializer)
 	, ParentSubmix(nullptr)
+{}
+
+USoundSubmixBase::USoundSubmixBase(const FObjectInitializer& ObjectInitializer)
 #if WITH_EDITORONLY_DATA
-	, SoundSubmixGraph(nullptr)
+	: SoundSubmixGraph(nullptr)
 #endif // WITH_EDITORONLY_DATA
-	, ChannelFormat(ESubmixChannelFormat::Device)
+{}
+
+USoundSubmix::USoundSubmix(const FObjectInitializer& ObjectInitializer)
+	: Super(ObjectInitializer)
 	, bMuteWhenBackgrounded(0)
 	, AmbisonicsPluginSettings(nullptr)
 	, EnvelopeFollowerAttackTime(10)
@@ -28,6 +42,24 @@
 {
 }
 
+
+UEndpointSubmix::UEndpointSubmix(const FObjectInitializer& ObjectInitializer)
+	: Super(ObjectInitializer)
+	, EndpointType(IAudioEndpointFactory::GetTypeNameForDefaultEndpoint())
+{
+
+}
+
+USoundfieldSubmix::USoundfieldSubmix(const FObjectInitializer& ObjectInitializer)
+	: Super(ObjectInitializer)
+	, SoundfieldEncodingFormat(ISoundfieldFactory::GetFormatNameForInheritedEncoding())
+{}
+
+USoundfieldEndpointSubmix::USoundfieldEndpointSubmix(const FObjectInitializer& ObjectInitializer)
+	: Super(ObjectInitializer)
+	, SoundfieldEndpointType(ISoundfieldEndpointFactory::DefaultSoundfieldEndpointName())
+{}
+
 void USoundSubmix::StartRecordingOutput(const UObject* WorldContextObject, float ExpectedDuration)
 {
 	if (!GEngine)
@@ -204,13 +236,33 @@
 		AudioDevice->SetSubmixOutputVolume(this, InOutputVolume);
 	}
 }
+#if WITH_EDITOR
+void USoundSubmix::PostEditChangeProperty(struct FPropertyChangedEvent& PropertyChangedEvent)
+{
+	if (PropertyChangedEvent.Property != nullptr)
+	{
+		static const FName NAME_OutputVolume(TEXT("OutputVolume"));
 
-FString USoundSubmix::GetDesc()
+		if (PropertyChangedEvent.Property->GetFName() == NAME_OutputVolume)
+		{
+			FAudioDeviceManager* AudioDeviceManager = (GEngine ? GEngine->GetAudioDeviceManager() : nullptr);
+			if (AudioDeviceManager)
+			{
+				AudioDeviceManager->UpdateSubmix(this);
+			}
+		}
+	}
+
+	Super::PostEditChangeProperty(PropertyChangedEvent);
+}
+#endif
+
+FString USoundSubmixBase::GetDesc()
 {
-	return FString(TEXT("Sound submix"));
+	return FString(TEXT("Sound Submix"));
 }
 
-void USoundSubmix::BeginDestroy()
+void USoundSubmixBase::BeginDestroy()
 {
 	Super::BeginDestroy();
 
@@ -224,10 +276,37 @@
 	}
 }
 
-void USoundSubmix::PostLoad()
+void USoundSubmixBase::PostLoad()
 {
 	Super::PostLoad();
 
+	if (ClearBrokenSubmixAssetsCVar)
+	{
+		for (int32 ChildIndex = ChildSubmixes.Num() - 1; ChildIndex >= 0; ChildIndex--)
+		{
+			USoundSubmixBase* ChildSubmix = ChildSubmixes[ChildIndex];
+
+			if (!ChildSubmix)
+			{
+				continue;
+			}
+
+			if (USoundSubmixWithParentBase* CastedChildSubmix = Cast<USoundSubmixWithParentBase>(ChildSubmix))
+			{
+				if (!ensure(CastedChildSubmix->ParentSubmix == this))
+				{
+					UE_LOG(LogAudio, Warning, TEXT("Submix had a child submix that didn't explicitly mark this submix as a parent!"));
+					ChildSubmixes.RemoveAtSwap(ChildIndex);
+				}
+			}
+			else
+			{
+				ensureMsgf(false, TEXT("Submix had a child submix that doesn't have an output!"));
+				ChildSubmixes.RemoveAtSwap(ChildIndex);
+			}
+		}
+	}
+
 	// Use the main/default audio device for storing and retrieving sound class properties
 	FAudioDeviceManager* AudioDeviceManager = (GEngine ? GEngine->GetAudioDeviceManager() : nullptr);
 
@@ -239,7 +318,16 @@
 }
 
 #if WITH_EDITOR
-void USoundSubmix::PreEditChange(FProperty* PropertyAboutToChange)
+
+void USoundSubmixBase::PostDuplicate(EDuplicateMode::Type DuplicateMode)
 {
+	if (DuplicateMode == EDuplicateMode::Normal)
+	{
+		ChildSubmixes.Reset();
+	}
+}
+
+void USoundSubmixBase::PreEditChange(FProperty* PropertyAboutToChange)
+{
 	static FName NAME_ChildSubmixes(TEXT("ChildSubmixes"));
 
@@ -250,7 +338,7 @@
 	}
 }
 
-void USoundSubmix::PostEditChangeProperty(FPropertyChangedEvent& PropertyChangedEvent)
+void USoundSubmixBase::PostEditChangeProperty(FPropertyChangedEvent& PropertyChangedEvent)
 {
 	// Whether or not we need to reinit the submix. Not all properties require reinitialization.
 	bool bReinitSubmix = true;
@@ -279,10 +367,10 @@
 						// Revert to the child submixes
 						ChildSubmixes = BackupChildSubmixes;
 					}
-					else
+					else if (USoundSubmixWithParentBase* SubmixWithParent = CastChecked<USoundSubmixWithParentBase>(ChildSubmixes[ChildIndex]))
 					{
 						// Update parentage
-						ChildSubmixes[ChildIndex]->SetParentSubmix(this);
+						SubmixWithParent->SetParentSubmix(this);
 					}
 					break;
 				}
@@ -294,11 +382,75 @@
 				if (BackupChildSubmixes[ChildIndex] != nullptr && !ChildSubmixes.Contains(BackupChildSubmixes[ChildIndex]))
 				{
 					BackupChildSubmixes[ChildIndex]->Modify();
-					BackupChildSubmixes[ChildIndex]->ParentSubmix = nullptr;
+					if (USoundSubmixWithParentBase* SubmixWithParent = Cast<USoundSubmixWithParentBase>(BackupChildSubmixes[ChildIndex]))
+					{
+						SubmixWithParent->ParentSubmix = nullptr;
+					}
 				}
 			}
 		}
-		else if (PropertyChangedEvent.Property->GetFName() == NAME_ParentSubmix)
+	}
+
+	if (GEngine)
+	{
+		// Force the properties to be initialized for this SoundSubmix on all active audio devices
+		if (FAudioDeviceManager* AudioDeviceManager = GEngine->GetAudioDeviceManager())
 		{
+			AudioDeviceManager->RegisterSoundSubmix(this);
+		}
+	}
+
+	BackupChildSubmixes.Reset();
+
+	Super::PostEditChangeProperty(PropertyChangedEvent);
+}
+
+TArray<USoundSubmixBase*> USoundSubmixBase::BackupChildSubmixes;
+
+bool USoundSubmixBase::RecurseCheckChild(const USoundSubmixBase* ChildSoundSubmix) const
+{
+	for (int32 Index = 0; Index < ChildSubmixes.Num(); Index++)
+	{
+		if (ChildSubmixes[Index])
+		{
+			if (ChildSubmixes[Index] == ChildSoundSubmix)
+			{
+				return true;
+			}
+
+			if (ChildSubmixes[Index]->RecurseCheckChild(ChildSoundSubmix))
+			{
+				return true;
+			}
+		}
+	}
+
+	return false;
+}
+
+void USoundSubmixWithParentBase::SetParentSubmix(USoundSubmixBase* InParentSubmix)
+{
+	if (ParentSubmix != InParentSubmix)
+	{
+		if (ParentSubmix != nullptr)
+		{
+			ParentSubmix->Modify();
+			ParentSubmix->ChildSubmixes.Remove(this);
+		}
+
+		Modify();
+		ParentSubmix = InParentSubmix;
+		ParentSubmix->ChildSubmixes.AddUnique(this);
+	}
+}
+
+void USoundSubmixWithParentBase::PostEditChangeProperty(struct FPropertyChangedEvent& PropertyChangedEvent)
+{
+	if (PropertyChangedEvent.Property != nullptr)
+	{
+		static const FName NAME_ParentSubmix(TEXT("ParentSubmix"));
+
+		if (PropertyChangedEvent.Property->GetFName() == NAME_ParentSubmix)
+		{
 			// Add this sound class to the parent class if it's not already added
 			if (ParentSubmix)
@@ -306,7 +458,7 @@
 				bool bIsChildSubmix = false;
 				for (int32 i = 0; i < ParentSubmix->ChildSubmixes.Num(); ++i)
 				{
-					USoundSubmix* ChildSubmix = ParentSubmix->ChildSubmixes[i];
+					USoundSubmixBase* ChildSubmix = ParentSubmix->ChildSubmixes[i];
 					if (ChildSubmix && ChildSubmix == this)
 					{
 						bIsChildSubmix = true;
@@ -323,18 +475,9 @@
 
 			Modify();
 		}
-		else if (PropertyChangedEvent.Property->GetFName() == NAME_OutputVolume)
-		{
-			FAudioDeviceManager* AudioDeviceManager = (GEngine ? GEngine->GetAudioDeviceManager() : nullptr);
-			if (AudioDeviceManager)
-			{
-				AudioDeviceManager->UpdateSubmix(this);
-			}
-			bReinitSubmix = false;
-		}
 	}
 
-	if (GEngine && bReinitSubmix)
+	if (GEngine)
 	{
 		// Force the properties to be initialized for this SoundSubmix on all active audio devices
 		if (FAudioDeviceManager* AudioDeviceManager = GEngine->GetAudioDeviceManager())
@@ -343,70 +486,26 @@
 		}
 	}
 
-	BackupChildSubmixes.Reset();
-
 	Super::PostEditChangeProperty(PropertyChangedEvent);
 }
 
-void USoundSubmix::PostDuplicate(EDuplicateMode::Type DuplicateMode)
+void USoundSubmixWithParentBase::PostDuplicate(EDuplicateMode::Type DuplicateMode)
 {
 	if (DuplicateMode == EDuplicateMode::Normal)
 	{
 		SetParentSubmix(nullptr);
-		ChildSubmixes.Reset();
 	}
 
 	Super::PostDuplicate(DuplicateMode);
 }
 
-bool USoundSubmix::RecurseCheckChild(const USoundSubmix* ChildSoundSubmix) const
+void USoundSubmixBase::AddReferencedObjects(UObject* InThis, FReferenceCollector& Collector)
 {
-	for (int32 Index = 0; Index < ChildSubmixes.Num(); Index++)
-	{
-		if (ChildSubmixes[Index])
-		{
-			if (ChildSubmixes[Index] == ChildSoundSubmix)
-			{
-				return true;
-			}
+	USoundSubmixBase* This = CastChecked<USoundSubmixBase>(InThis);
 
-			if (ChildSubmixes[Index]->RecurseCheckChild(ChildSoundSubmix))
-			{
-				return true;
-			}
-		}
-	}
-
-	return false;
-}
-
-void USoundSubmix::SetParentSubmix(USoundSubmix* InParentSubmix)
-{
-	if (ParentSubmix != InParentSubmix)
-	{
-		if (ParentSubmix)
-		{
-			ParentSubmix->Modify();
-			ParentSubmix->ChildSubmixes.Remove(this);
-		}
-
-		Modify();
-
-		ParentSubmix = InParentSubmix;
-		if (ParentSubmix)
-		{
-			ParentSubmix->ChildSubmixes.AddUnique(this);
-		}
-	}
-}
-
-void USoundSubmix::AddReferencedObjects(UObject* InThis, FReferenceCollector& Collector)
-{
-	USoundSubmix* This = CastChecked<USoundSubmix>(InThis);
-
 	Collector.AddReferencedObject(This->SoundSubmixGraph, This);
 
-	for (USoundSubmix* Backup : This->BackupChildSubmixes)
+	for (USoundSubmixBase* Backup : This->BackupChildSubmixes)
 	{
 		Collector.AddReferencedObject(Backup);
 	}
@@ -414,3 +513,110 @@
 	Super::AddReferencedObjects(InThis, Collector);
 }
 #endif // WITH_EDITOR
+
+ISoundfieldFactory* USoundfieldSubmix::GetSoundfieldFactoryForSubmix() const
+{
+	// If this isn't called in the game thread, a ParentSubmix could get destroyed while we are recursing through the submix graph.
+	ensure(IsInGameThread());
+
+	FName SoundfieldFormat = GetSubmixFormat();
+	check(SoundfieldFormat != ISoundfieldFactory::GetFormatNameForInheritedEncoding());
+
+	return ISoundfieldFactory::Get(SoundfieldFormat);
+}
+
+const USoundfieldEncodingSettingsBase* USoundfieldSubmix::GetSoundfieldEncodingSettings() const
+{
+	return GetEncodingSettings();
+}
+
+TArray<USoundfieldEffectBase *> USoundfieldSubmix::GetSoundfieldProcessors() const
+{
+	return SoundfieldEffectChain;
+}
+
+FName USoundfieldSubmix::GetSubmixFormat() const
+{
+	USoundfieldSubmix* ParentSoundfieldSubmix = Cast<USoundfieldSubmix>(ParentSubmix);
+
+	if (!ParentSoundfieldSubmix || SoundfieldEncodingFormat != ISoundfieldFactory::GetFormatNameForInheritedEncoding())
+	{
+		if (SoundfieldEncodingFormat == ISoundfieldFactory::GetFormatNameForInheritedEncoding())
+		{
+			return ISoundfieldFactory::GetFormatNameForNoEncoding();
+		}
+		else
+		{
+			return SoundfieldEncodingFormat;
+		}
+
+	}
+	else if(ParentSoundfieldSubmix)
+	{
+		// If this submix matches the format of whatever submix it's plugged into, 
+		// Recurse into the submix graph to find it.
+		return ParentSoundfieldSubmix->GetSubmixFormat();
+	}
+	else
+	{
+		return ISoundfieldFactory::GetFormatNameForNoEncoding();
+	}
+}
+
+const USoundfieldEncodingSettingsBase* USoundfieldSubmix::GetEncodingSettings() const
+{
+	FName SubmixFormatName = GetSubmixFormat();
+
+	USoundfieldSubmix* ParentSoundfieldSubmix = Cast<USoundfieldSubmix>(ParentSubmix);
+
+	if (EncodingSettings)
+	{
+		return EncodingSettings;
+	}
+	else if (ParentSoundfieldSubmix && SoundfieldEncodingFormat == ISoundfieldFactory::GetFormatNameForInheritedEncoding())
+	{
+		// If this submix matches the format of whatever it's plugged into,
+		// Recurse into the submix graph to match it's settings.
+		return ParentSoundfieldSubmix->GetEncodingSettings();
+	}
+	else if (ISoundfieldFactory* Factory = ISoundfieldFactory::Get(SubmixFormatName))
+	{
+		// If we don't have any encoding settings, use the default.
+		return Factory->GetDefaultEncodingSettings();
+	}
+	else
+	{
+		// If we don't have anything, exit.
+		return nullptr;
+	}
+}
+
+IAudioEndpointFactory* UEndpointSubmix::GetAudioEndpointForSubmix() const
+{
+	return IAudioEndpointFactory::Get(EndpointType);
+}
+
+const UAudioEndpointSettingsBase* UEndpointSubmix::GetEndpointSettings() const
+{
+	return EndpointSettings;
+}
+
+ISoundfieldEndpointFactory* USoundfieldEndpointSubmix::GetSoundfieldEndpointForSubmix() const
+{
+	return ISoundfieldEndpointFactory::Get(SoundfieldEndpointType);
+}
+
+const USoundfieldEndpointSettingsBase* USoundfieldEndpointSubmix::GetEndpointSettings() const
+{
+	return EndpointSettings;
+}
+
+const USoundfieldEncodingSettingsBase* USoundfieldEndpointSubmix::GetEncodingSettings() const
+{
+	return EncodingSettings;
+}
+
+TArray<USoundfieldEffectBase*> USoundfieldEndpointSubmix::GetSoundfieldProcessors() const
+{
+	return SoundfieldEffectChain;
+}
diff --git a/Engine/Source/Runtime/Engine/Private/SoundWave.cpp b/Engine/Source/Runtime/Engine/Private/SoundWave.cpp
--- a/Engine/Source/Runtime/Engine/Private/SoundWave.cpp
+++ b/Engine/Source/Runtime/Engine/Private/SoundWave.cpp
@@ -1822,16 +1822,12 @@ void USoundWave::Parse(FAudioDevice* AudioDevice, const UPTRINT NodeWaveInstance
 	}
 
 	WaveInstance->bIsAmbisonics = bIsAmbisonics;
-	if (bIsAmbisonics)
-	{
-		UAudioSettings* Settings = GetMutableDefault<UAudioSettings>();
-		WaveInstance->SoundSubmix = Cast<USoundSubmix>(FSoftObjectPtr(Settings->AmbisonicSubmix).Get());
-	}
-	else if (ParseParams.SoundSubmix)
+
+	if (ParseParams.SoundSubmix)
 	{
 		WaveInstance->SoundSubmix = ParseParams.SoundSubmix;
 	}
-	else if (USoundSubmix* WaveSubmix = GetSoundSubmix())
+	else if (USoundSubmixBase* WaveSubmix = GetSoundSubmix())
 	{
 		WaveInstance->SoundSubmix = WaveSubmix;
 	}
diff --git a/Engine/Source/Runtime/Engine/Public/ActiveSound.h b/Engine/Source/Runtime/Engine/Public/ActiveSound.h
--- a/Engine/Source/Runtime/Engine/Public/ActiveSound.h
+++ b/Engine/Source/Runtime/Engine/Public/ActiveSound.h
@@ -126,7 +126,7 @@ struct FSoundParseParameters
 	float AbsoluteAzimuth;
 
 	// The sound submix to use for the wave instance
-	USoundSubmix* SoundSubmix;
+	USoundSubmixBase* SoundSubmix;
 
 	// The submix sends to use
 	TArray<FSoundSubmixSendInfo> SoundSubmixSends;
@@ -633,7 +633,7 @@ public:
 	/**
 	* Get the sound submix to use for this sound instance
 	*/
-	USoundSubmix* GetSoundSubmix() const;
+	USoundSubmixBase* GetSoundSubmix() const;
 
 	/** Gets the sound submix sends to use for this sound instance. */
 	void GetSoundSubmixSends(TArray<FSoundSubmixSendInfo>& OutSends) const;
diff --git a/Engine/Source/Runtime/Engine/Public/Audio.h b/Engine/Source/Runtime/Engine/Public/Audio.h
--- a/Engine/Source/Runtime/Engine/Public/Audio.h
+++ b/Engine/Source/Runtime/Engine/Public/Audio.h
@@ -166,7 +166,7 @@ struct ENGINE_API FWaveInstance
 	FSoundModulationControls SoundModulationControls;
 
 	/** Sound submix object to send audio to for mixing in audio mixer.  */
-	USoundSubmix* SoundSubmix;
+	USoundSubmixBase* SoundSubmix;
 
 	/** Sound submix sends */
 	TArray<FSoundSubmixSendInfo> SoundSubmixSends;
diff --git a/Engine/Source/Runtime/Engine/Public/AudioDevice.h b/Engine/Source/Runtime/Engine/Public/AudioDevice.h
--- a/Engine/Source/Runtime/Engine/Public/AudioDevice.h
+++ b/Engine/Source/Runtime/Engine/Public/AudioDevice.h
@@ -44,7 +44,7 @@ class USoundConcurrency;
 class USoundEffectSourcePreset;
 class USoundEffectSubmixPreset;
 class USoundMix;
-class USoundSubmix;
+class USoundSubmixBase;
 class USoundSourceBus;
 class USoundWave;
 class UWorld;
@@ -661,11 +661,6 @@ public:
 
 	const TArray<FListener>& GetListeners() const { check(IsInAudioThread()); return Listeners; }
 
-	/**
-	* Get ambisonics mixer, if one is available
-	*/
-	TAmbisonicsMixerPtr GetAmbisonicsMixer() { return AmbisonicsMixer; };
-
 	/**
 	 * Returns the currently applied reverb effect if there is one.
 	 */
@@ -830,10 +825,10 @@ public:
 	virtual void InitSoundSubmixes() {}
 
 	/** Registers the sound submix */
-	virtual void RegisterSoundSubmix(USoundSubmix* SoundSubmix, bool bInit) {}
+	virtual void RegisterSoundSubmix(const USoundSubmixBase* SoundSubmix, bool bInit) {}
 
 	/** Unregisters the sound submix */
-	virtual void UnregisterSoundSubmix(USoundSubmix* SoundSubmix) {}
+	virtual void UnregisterSoundSubmix(const USoundSubmixBase* SoundSubmix) {}
 
 	/**
 	 * Registers the submix buffer listener with the given submix.
@@ -1244,7 +1239,7 @@ public:
 	virtual bool GetCurrentSourceEffectChain(const uint32 SourceEffectChainId, TArray<FSourceEffectChainEntry>& OutCurrentSourceEffectChainEntries) { return false; }
 
 	/** Updates the submix properties of any playing submix instances. Allows editor to make changes to submix properties and hear them propagate live.*/
-	virtual void UpdateSubmixProperties(USoundSubmix* InSubmix)
+	virtual void UpdateSubmixProperties(USoundSubmixBase* InSubmix)
 	{
 		UE_LOG(LogAudio, Error, TEXT("Submixes are only supported in audio mixer."));
 	}
@@ -1724,9 +1719,6 @@ public:
 	/** 3rd party modulation interface */
 	TAudioModulationPtr ModulationInterface;
 
-	/** This devices ambisonics pointer, if one exists */
-	TAmbisonicsMixerPtr AmbisonicsMixer;
-
 	/** 3rd party listener observers registered to this audio device. */
 	TArray<TAudioPluginListenerPtr> PluginListeners;
 
diff --git a/Engine/Source/Runtime/Engine/Public/AudioDeviceManager.h b/Engine/Source/Runtime/Engine/Public/AudioDeviceManager.h
--- a/Engine/Source/Runtime/Engine/Public/AudioDeviceManager.h
+++ b/Engine/Source/Runtime/Engine/Public/AudioDeviceManager.h
@@ -22,7 +22,7 @@ class IAudioDeviceModule;
 class UAudioComponent;
 class USoundClass;
 class USoundMix;
-class USoundSubmix;
+class USoundSubmixBase;
 class USoundWave;
 class UWorld;
 class FAudioDevice;
@@ -300,10 +300,10 @@ public:
 	void InitSoundClasses();
 
 	/** Registers the Sound Mix for all active devices. */
-	void RegisterSoundSubmix(USoundSubmix* SoundSubmix);
+	void RegisterSoundSubmix(const USoundSubmixBase* SoundSubmix);
 
 	/** Registers the Sound Mix for all active devices. */
-	void UnregisterSoundSubmix(USoundSubmix* SoundSubmix);
+	void UnregisterSoundSubmix(const USoundSubmixBase* SoundSubmix);
 
 	/** Initializes the sound mixes for all active devices. */
 	void InitSoundSubmixes();
@@ -315,7 +315,7 @@ public:
 	void UpdateSourceEffectChain(const uint32 SourceEffectChainId, const TArray<FSourceEffectChainEntry>& SourceEffectChain, const bool bPlayEffectChainTails);
 
 	/** Updates this submix for any changes made. Broadcasts to all submix instances. */
-	void UpdateSubmix(USoundSubmix* SoundSubmix);
+	void UpdateSubmix(USoundSubmixBase* SoundSubmix);
 
 	/** Sets which audio device is the active audio device. */
 	void SetActiveDevice(uint32 InAudioDeviceHandle);
diff --git a/Engine/Source/Runtime/SignalProcessing/Private/BufferVectorOperations.cpp b/Engine/Source/Runtime/SignalProcessing/Private/BufferVectorOperations.cpp
--- a/Engine/Source/Runtime/SignalProcessing/Private/BufferVectorOperations.cpp
+++ b/Engine/Source/Runtime/SignalProcessing/Private/BufferVectorOperations.cpp
@@ -6,6 +6,12 @@
 
 namespace Audio
 {
+	static void RestrictedPtrAliasCheck(const float* RESTRICT Ptr1, const float* RESTRICT Ptr2, uint32 NumFloatsInArray)
+	{
+		checkf(FMath::Abs(Ptr1 - Ptr2) >= NumFloatsInArray,
+			TEXT("Using this function as an in-place operation will result in undefined behavior!"));
+	}
+
 	/* Sets a values to zero if value is denormal. Denormal numbers significantly slow down floating point operations. */
 	void BufferUnderflowClampFast(AlignedFloatBuffer& InOutBuffer)
 	{
@@ -89,6 +95,8 @@ namespace Audio
 	void BufferMultiplyByConstant(const float* RESTRICT InFloatBuffer, float InValue, float* RESTRICT OutFloatBuffer, const int32 InNumSamples)
 	{
 		check(InNumSamples >= 4);
+		RestrictedPtrAliasCheck(InFloatBuffer, OutFloatBuffer, InNumSamples);
+
 #if !AUDIO_USE_SIMD
 		for (int32 i = 0; i < InNumSamples; ++i)
 		{
@@ -186,17 +194,20 @@ namespace Audio
 		checkf(IsAligned<const float*>(InBuffer2, AUDIO_SIMD_FLOAT_ALIGNMENT), TEXT("Memory must be aligned to use vector operations."));
 		checkf(IsAligned<float*>(OutBuffer, AUDIO_SIMD_FLOAT_ALIGNMENT), TEXT("Memory must be aligned to use vector operations."));
 
+		RestrictedPtrAliasCheck(InBuffer1, OutBuffer, InNum);
+		RestrictedPtrAliasCheck(InBuffer2, OutBuffer, InNum);
+
 		VectorRegister Gain1Vector = VectorLoadFloat1(&InGain1);
 		VectorRegister Gain2Vector = VectorLoadFloat1(&InGain2);
 
 		for (int32 i = 0; i < InNum; i += 4)
 		{
 			// InBuffer1 x InGain1
-			VectorRegister Input1  = VectorLoadAligned(&InBuffer1[i]);
+			VectorRegister Input1 = VectorLoadAligned(&InBuffer1[i]);
 			VectorRegister Weighted1 = VectorMultiply(Input1, Gain1Vector);
 
 			// InBuffer2 x InGain2
-			VectorRegister Input2  = VectorLoadAligned(&InBuffer2[i]);
+			VectorRegister Input2 = VectorLoadAligned(&InBuffer2[i]);
 			VectorRegister Weighted2 = VectorMultiply(Input2, Gain2Vector);
 
 			VectorRegister Output = VectorAdd(Weighted1, Weighted2);
@@ -213,15 +224,18 @@ namespace Audio
 		checkf(IsAligned<const float*>(InBuffer2, AUDIO_SIMD_FLOAT_ALIGNMENT), TEXT("Memory must be aligned to use vector operations."));
 		checkf(IsAligned<float*>(OutBuffer, AUDIO_SIMD_FLOAT_ALIGNMENT), TEXT("Memory must be aligned to use vector operations."));
 
+		RestrictedPtrAliasCheck(InBuffer1, OutBuffer, InNum);
+		RestrictedPtrAliasCheck(InBuffer2, OutBuffer, InNum);
+
 		VectorRegister Gain1Vector = VectorLoadFloat1(&InGain1);
 
 		for (int32 i = 0; i < InNum; i += 4)
 		{
 			// InBuffer1 x InGain1
-			VectorRegister Input1  = VectorLoadAligned(&InBuffer1[i]);
+			VectorRegister Input1 = VectorLoadAligned(&InBuffer1[i]);
 			VectorRegister Weighted1 = VectorMultiply(Input1, Gain1Vector);
 
-			VectorRegister Input2  = VectorLoadAligned(&InBuffer2[i]);
+			VectorRegister Input2 = VectorLoadAligned(&InBuffer2[i]);
 
 			VectorRegister Output = VectorAdd(Weighted1, Input2);
 			VectorStoreAligned(Output, &OutBuffer[i]);
@@ -317,6 +331,8 @@ namespace Audio
 		checkf(IsAligned<float*>(BufferToSumTo, 4), TEXT("Memory must be aligned to use vector operations."));
 		checkf(NumSamples % 4 == 0, TEXT("Please use a buffer size that is a multiple of 4."));
 
+		RestrictedPtrAliasCheck(InFloatBuffer, BufferToSumTo, NumSamples);
+
 #if !AUDIO_USE_SIMD
 		for (int32 i = 0; i < NumSamples; ++i)
 		{
@@ -347,6 +363,8 @@ namespace Audio
 		checkf(IsAligned<float*>(BufferToSumTo, 4), TEXT("Memory must be aligned to use vector operations."));
 		checkf(NumSamples % 4 == 0, TEXT("Please use a buffer size that is a multiple of 4."));
 
+		RestrictedPtrAliasCheck(InFloatBuffer, BufferToSumTo, NumSamples);
+
 #if !AUDIO_USE_SIMD
 		for (int32 i = 0; i < NumSamples; ++i)
 		{
@@ -374,6 +392,8 @@ namespace Audio
 		checkf(IsAligned<float*>(BufferToSumTo, 4), TEXT("Memory must be aligned to use vector operations."));
 		checkf(NumSamples % 4 == 0, TEXT("Please use a buffer size that is a multiple of 4."));
 
+		RestrictedPtrAliasCheck(InFloatBuffer, BufferToSumTo, NumSamples);
+
 		const int32 NumIterations = NumSamples / 4;
 
 		if (FMath::IsNearlyEqual(StartGain, EndGain))
@@ -441,6 +461,9 @@ namespace Audio
 		checkf(IsAligned<const float*>(InSubtrahend, AUDIO_SIMD_FLOAT_ALIGNMENT), TEXT("Memory must be aligned to use vector operations."));
 		checkf(IsAligned<float*>(OutBuffer, AUDIO_SIMD_FLOAT_ALIGNMENT), TEXT("Memory must be aligned to use vector operations."));
 
+		RestrictedPtrAliasCheck(InMinuend, OutBuffer, InNum);
+		RestrictedPtrAliasCheck(InSubtrahend, OutBuffer, InNum);
+
 		for (int32 i = 0; i < InNum; i += 4)
 		{
 			VectorRegister Input1 = VectorLoadAligned(&InMinuend[i]);
@@ -465,6 +488,8 @@ namespace Audio
 		checkf(IsAligned<const float*>(InMinuend, AUDIO_SIMD_FLOAT_ALIGNMENT), TEXT("Memory must be aligned to use vector operations."));
 		checkf(IsAligned<const float*>(InOutSubtrahend, AUDIO_SIMD_FLOAT_ALIGNMENT), TEXT("Memory must be aligned to use vector operations."));
 		
+		RestrictedPtrAliasCheck(InMinuend, InOutSubtrahend, InNum);
+
 		for (int32 i = 0; i < InNum; i += 4)
 		{
 			VectorRegister Input1 = VectorLoadAligned(&InMinuend[i]);
@@ -490,6 +515,8 @@ namespace Audio
 		checkf(IsAligned<const float*>(InOutMinuend, AUDIO_SIMD_FLOAT_ALIGNMENT), TEXT("Memory must be aligned to use vector operations."));
 		checkf(IsAligned<const float*>(InSubtrahend, AUDIO_SIMD_FLOAT_ALIGNMENT), TEXT("Memory must be aligned to use vector operations."));
 
+		RestrictedPtrAliasCheck(InOutMinuend, InSubtrahend, InNum);
+		
 		for (int32 i = 0; i < InNum; i += 4)
 		{
 			VectorRegister Input1 = VectorLoadAligned(&InOutMinuend[i]);
@@ -517,6 +544,9 @@ namespace Audio
 		checkf(IsAligned<float*>(OutputBuffer, 4), TEXT("Memory must be aligned to use vector operations."));
 		checkf(NumSamples % 4 == 0, TEXT("Please use a buffer size that is a multiple of 4."));
 
+		RestrictedPtrAliasCheck(InFloatBuffer1, OutputBuffer, NumSamples);
+		RestrictedPtrAliasCheck(InFloatBuffer2, OutputBuffer, NumSamples);
+
 #if !AUDIO_USE_SIMD
 		for (int32 i = 0; i < NumSamples; ++i)
 		{
@@ -545,6 +575,8 @@ namespace Audio
 		checkf(IsAligned<const float*>(BufferToMultiply, 4), TEXT("Memory must be aligned to use vector operations."));
 		checkf(NumSamples % 4 == 0, TEXT("Please use a buffer size that is a multiple of 4."));
 
+		RestrictedPtrAliasCheck(InFloatBuffer, BufferToMultiply, NumSamples);
+
 		for (int32 i = 0; i < NumSamples; i += 4)
 		{
 			VectorRegister Input1 = VectorLoadAligned(&InFloatBuffer[i]);
@@ -1774,6 +1806,28 @@ namespace Audio
 	}
 
 
+	SIGNALPROCESSING_API void DownmixAndSumIntoBuffer(int32 NumSourceChannels, int32 NumDestinationChannels, const AlignedFloatBuffer& SourceBuffer, AlignedFloatBuffer& BufferToSumTo, const float* RESTRICT Gains)
+	{
+		DownmixAndSumIntoBuffer(NumSourceChannels, NumDestinationChannels, SourceBuffer.GetData(), BufferToSumTo.GetData(), SourceBuffer.Num() / NumSourceChannels, Gains);
+	}
+
+	SIGNALPROCESSING_API void DownmixAndSumIntoBuffer(int32 NumSourceChannels, int32 NumDestinationChannels, const float* RESTRICT SourceBuffer, float* RESTRICT BufferToSumTo, int32 NumFrames, const float* RESTRICT Gains)
+	{
+		for (int32 FrameIndex = 0; FrameIndex < NumFrames; FrameIndex++)
+		{
+			float* RESTRICT OutputFrame = &BufferToSumTo[FrameIndex * NumDestinationChannels];
+			const float* RESTRICT InputFrame = &SourceBuffer[FrameIndex * NumSourceChannels];
+
+			for (int32 OutputChannelIndex = 0; OutputChannelIndex < NumDestinationChannels; OutputChannelIndex++)
+			{
+				for (int32 InputChannelIndex = 0; InputChannelIndex < NumSourceChannels; InputChannelIndex++)
+				{
+					OutputFrame[OutputChannelIndex] += InputFrame[InputChannelIndex] * Gains[InputChannelIndex * NumDestinationChannels + OutputChannelIndex];
+				}
+			}
+		}
+	}
+
 	/** Interleaves samples from two input buffers */
 	void BufferInterleave2ChannelFast(const AlignedFloatBuffer& InBuffer1, const AlignedFloatBuffer& InBuffer2, AlignedFloatBuffer& OutBuffer)
 	{
@@ -1932,7 +1986,9 @@ namespace Audio
 		checkf(IsAligned<const float*>(InImaginarySamples, AUDIO_SIMD_FLOAT_ALIGNMENT), TEXT("Memory must be aligned to use vector operations."));
 		checkf(IsAligned<float*>(OutPowerSamples, AUDIO_SIMD_FLOAT_ALIGNMENT), TEXT("Memory must be aligned to use vector operations."));
 
-		
+		RestrictedPtrAliasCheck(InRealSamples, OutPowerSamples, InNum);
+		RestrictedPtrAliasCheck(InImaginarySamples, OutPowerSamples, InNum);
+
 		for (int32 i = 0; i < InNum; i += 4)
 		{
 			VectorRegister VInReal = VectorLoadAligned(&InRealSamples[i]);
@@ -1946,4 +2002,74 @@ namespace Audio
 			VectorStoreAligned(VOut, &OutPowerSamples[i]);
 		}
 	}
+
+	// class FBufferLinearEase implementation
+	FBufferLinearEase::FBufferLinearEase() {}
+	FBufferLinearEase::FBufferLinearEase(const AlignedFloatBuffer & InSourceValues, const AlignedFloatBuffer & InTargetValues, int32 InLerpLength)
+	{
+		Init(InSourceValues, InTargetValues, InLerpLength);
+	}
+
+	FBufferLinearEase::~FBufferLinearEase() {}
+
+	void FBufferLinearEase::Init(const AlignedFloatBuffer & InSourceValues, const AlignedFloatBuffer & InTargetValues, int32 InLerpLength)
+	{
+		check(InSourceValues.Num());
+		check(InTargetValues.Num());
+		check(InLerpLength > 0);
+
+		BufferLength = InSourceValues.Num();
+
+		check(InTargetValues.Num() == BufferLength);
+		LerpLength = InLerpLength;
+		CurrentLerpStep = 0;
+
+		// init deltas
+		DeltaBuffer.Reset();
+		DeltaBuffer.AddZeroed(BufferLength);
+
+		const float OneOverLerpLength = 1.0f / static_cast<float>(LerpLength);
+
+		BufferSubtractFast(InTargetValues.GetData(), InSourceValues.GetData(), DeltaBuffer.GetData(), BufferLength);
+		MultiplyBufferByConstantInPlace(DeltaBuffer, OneOverLerpLength);
+	}
+
+	bool FBufferLinearEase::Update(AlignedFloatBuffer & InSourceValues)
+	{
+		check(InSourceValues.Num() == BufferLength);
+		check(CurrentLerpStep != LerpLength);
+
+		MixInBufferFast(DeltaBuffer.GetData(), InSourceValues.GetData(), BufferLength);
+
+		if (++CurrentLerpStep == LerpLength)
+		{
+			return true;
+		}
+
+		return false;
+	}
+
+	bool FBufferLinearEase::Update(uint32 StepsToJumpForward, AlignedFloatBuffer & InSourceValues)
+	{
+		check(InSourceValues.Num() == BufferLength);
+		check(CurrentLerpStep != LerpLength);
+		check(StepsToJumpForward);
+
+		bool bIsComplete = false;
+		if ((CurrentLerpStep += StepsToJumpForward) >= LerpLength)
+		{
+			StepsToJumpForward -= (CurrentLerpStep - LerpLength);
+			CurrentLerpStep = LerpLength;
+			bIsComplete = true;
+		}
+
+		MixInBufferFast(DeltaBuffer.GetData(), InSourceValues.GetData(), BufferLength, static_cast<float>(StepsToJumpForward));
+
+		return bIsComplete;
+
+	}
+	const AlignedFloatBuffer & FBufferLinearEase::GetDeltaBuffer()
+	{
+		return DeltaBuffer;
+	}
 }
diff --git a/Engine/Source/Runtime/SignalProcessing/Private/MultithreadedPatching.cpp b/Engine/Source/Runtime/SignalProcessing/Private/MultithreadedPatching.cpp
--- a/Engine/Source/Runtime/SignalProcessing/Private/MultithreadedPatching.cpp
+++ b/Engine/Source/Runtime/SignalProcessing/Private/MultithreadedPatching.cpp
@@ -265,6 +265,12 @@ namespace Audio
 		}
 	}
 
+	void FPatchMixer::DisconnectAllInputs()
+	{
+		FScopeLock ScopeLock(&CurrentPatchesCriticalSection);
+		CurrentInputs.Reset();
+	}
+
 	void FPatchMixer::ConnectNewPatches()
 	{
 		FScopeLock ScopeLock(&PendingNewInputsCriticalSection);
diff --git a/Engine/Source/Runtime/SignalProcessing/Public/DSP/BufferVectorOperations.h b/Engine/Source/Runtime/SignalProcessing/Public/DSP/BufferVectorOperations.h
--- a/Engine/Source/Runtime/SignalProcessing/Public/DSP/BufferVectorOperations.h
+++ b/Engine/Source/Runtime/SignalProcessing/Public/DSP/BufferVectorOperations.h
@@ -237,6 +237,12 @@ namespace Audio
 	SIGNALPROCESSING_API void DownmixBuffer(int32 NumSourceChannels, int32 NumDestinationChannels, const AlignedFloatBuffer& SourceBuffer, AlignedFloatBuffer& DestinationBuffer, float* RESTRICT StartGains, const float* RESTRICT EndGains);
 	SIGNALPROCESSING_API void DownmixBuffer(int32 NumSourceChannels, int32 NumDestinationChannels, const float* RESTRICT SourceBuffer, float* RESTRICT DestinationBuffer, int32 NumFrames, float* RESTRICT StartGains, const float* RESTRICT EndGains);
 
+	/**
+	 * This is similar to DownmixBuffer, except that it sums into DestinationBuffer rather than overwriting it.
+	 */
+	SIGNALPROCESSING_API void DownmixAndSumIntoBuffer(int32 NumSourceChannels, int32 NumDestinationChannels, const AlignedFloatBuffer& SourceBuffer, AlignedFloatBuffer& BufferToSumTo, const float* RESTRICT Gains);
+	SIGNALPROCESSING_API void DownmixAndSumIntoBuffer(int32 NumSourceChannels, int32 NumDestinationChannels, const float* RESTRICT SourceBuffer, float* RESTRICT BufferToSumTo, int32 NumFrames, const float* RESTRICT Gains);
+
 	/** Interleaves samples from two input buffers */
 	SIGNALPROCESSING_API void BufferInterleave2ChannelFast(const AlignedFloatBuffer& InBuffer1, const AlignedFloatBuffer& InBuffer2, AlignedFloatBuffer& OutBuffer);
 
@@ -266,4 +272,32 @@ namespace Audio
 
 	/** Compute magnitude of complex data. Out[i] = Sqrt(Real[i] * Real[i] + Imaginary[i] * Imaginary[i]) */
 	SIGNALPROCESSING_API void BufferComplexToMagnitudeFast(const float* RESTRICT InRealSamples, const float* RESTRICT InImaginarySamples, float* RESTRICT OutPowerSamples, const int32 InNum);
+
+	/** Class which handles a vectorized interpolation of an entire buffer to the values of a target buffer */
+	class SIGNALPROCESSING_API FBufferLinearEase
+	{
+	public:
+		FBufferLinearEase();
+		FBufferLinearEase(const AlignedFloatBuffer& InSourceValues, const AlignedFloatBuffer& InTargetValues, int32 InLerpLength);
+		~FBufferLinearEase();
+
+		/** will cache SourceValues ptr and manually update SourceValues on Update() */
+		void Init(const AlignedFloatBuffer& InSourceValues, const AlignedFloatBuffer& InTargetValues, int32 InLerpLength);
+
+		/** Performs Vectorized update of SourceValues float buffer. Returns true if interpolation is complete */
+		bool Update(AlignedFloatBuffer& InSourceValues);
+
+		/** Update overloaded to let you jump forward more than a single time-step */
+		bool Update(uint32 StepsToJumpForward, AlignedFloatBuffer& InSourceValues);
+
+		/** returns const reference to the deltas buffer for doing interpolation elsewhere */
+		const AlignedFloatBuffer& GetDeltaBuffer();
+
+	private:
+		int32 BufferLength {0};
+		int32 LerpLength {0};
+		int32 CurrentLerpStep{0};
+		AlignedFloatBuffer DeltaBuffer;
+
+	}; // class BufferLerper
 }
diff --git a/Engine/Source/Runtime/SignalProcessing/Public/DSP/MultithreadedPatching.h b/Engine/Source/Runtime/SignalProcessing/Public/DSP/MultithreadedPatching.h
--- a/Engine/Source/Runtime/SignalProcessing/Public/DSP/MultithreadedPatching.h
+++ b/Engine/Source/Runtime/SignalProcessing/Public/DSP/MultithreadedPatching.h
@@ -137,6 +137,9 @@ namespace Audio
 		/** This function call gets the maximum number of samples that's safe to pop, based on the thread with the least amount of samples buffered. Thread safe, but blocks for PopAudio. */
 		int32 MaxNumberOfSamplesThatCanBePopped();
 
+		/** Disconnect everything currently connected to this mixer. */
+		void DisconnectAllInputs();
+
 	private:
 		/** Called within PopAudio. Flushes the PendingNewPatches array into CurrentPatches. During this function, AddNewPatch is blocked. */
 		void ConnectNewPatches();
diff --git a/Engine/Source/Runtime/SoundFieldRendering/Private/SphericalHarmonicCalculator.cpp b/Engine/Source/Runtime/SoundFieldRendering/Private/SphericalHarmonicCalculator.cpp
new file mode 100644
--- /dev/null
+++ b/Engine/Source/Runtime/SoundFieldRendering/Private/SphericalHarmonicCalculator.cpp
@@ -0,0 +1,352 @@
+// Copyright 1998-2019 Epic Games, Inc. All Rights Reserved.
+
+#pragma once
+#include "SphericalHarmonicCalculator.h"
+#include "DSP/Dsp.h"
+
+namespace SHBasis
+{
+#pragma region FOA
+	// zero order
+	static float SHBasisFunction_ACN_0(float, float)
+	{
+		return 1.0f;
+	}
+
+	// first order
+	static float SHBasisFunction_ACN_1(float Azimuth, float Elevation)
+	{
+		return FMath::Sin(Azimuth) * FMath::Cos(Elevation);
+	}
+
+	static float SHBasisFunction_ACN_2(float, float Elevation)
+	{
+		return FMath::Sin(Elevation);
+	}
+
+	static float SHBasisFunction_ACN_3(float Azimuth, float Elevation)
+	{
+		return FMath::Cos(Azimuth) * FMath::Cos(Elevation);
+	}
+#pragma  endregion
+#pragma region SECOND_ORDER
+	static const float MagicNumber_ACN_4 = FMath::Sqrt(3.0f) / 2.0f;
+	static float SHBasisFunction_ACN_4(const float Azimuth, const float Elevation)
+	{
+		const float CosElevation = FMath::Cos(Elevation);
+		return MagicNumber_ACN_4 * FMath::Sin(2.0f * Azimuth) * CosElevation*CosElevation;
+	}
+
+	static const float MagicNumber_ACN_5 = FMath::Sqrt(3.0f) / 2.0f;
+	static float SHBasisFunction_ACN_5(const float Azimuth, const float Elevation)
+	{
+		return MagicNumber_ACN_5 * FMath::Sin(Azimuth) * FMath::Sin(2.0f * Elevation);
+	}
+
+	static float SHBasisFunction_ACN_6(const float Azimuth, const float Elevation)
+	{
+		const float SinElevation = FMath::Sin(Elevation);
+		return 1.5f * SinElevation*SinElevation - 0.5f;
+	}
+
+	static const float MagicNumber_ACN_7 = FMath::Sqrt(3.0f) / 2.0f;
+	static float SHBasisFunction_ACN_7(const float Azimuth, const float Elevation)
+	{
+		return MagicNumber_ACN_7 * FMath::Cos(Azimuth) * FMath::Sin(2.0f * Elevation);
+	}
+
+	static const float MagicNumber_ACN_8 = FMath::Sqrt(3.0f) / 2.0f;
+	static float SHBasisFunction_ACN_8(const float Azimuth, const float Elevation)
+	{
+		static const float CosElevation = FMath::Cos(Elevation);
+		return MagicNumber_ACN_8 * FMath::Cos(2.07 * Azimuth) * CosElevation*CosElevation;
+	}
+	#pragma  endregion
+#pragma region THIRD_ORDER
+	static const float MagicNumber_ACN_9 = FMath::Sqrt(10.0f) / 4.0f;
+	static float SHBasisFunction_ACN_9(float Azimuth, float Elevation)
+	{
+		const float CosElevation = FMath::Cos(Elevation);
+		return  MagicNumber_ACN_9 * FMath::Sin(3.0f * Azimuth) * CosElevation*CosElevation*CosElevation;
+	}
+
+	static const float MagicNumber_ACN_10 = 0.5f * FMath::Sqrt(15.0f);
+	static float SHBasisFunction_ACN_10(float Azimuth, float Elevation)
+	{
+		static const float CosElevation = FMath::Cos(Elevation);
+		return MagicNumber_ACN_10 * FMath::Sin(2.0f * Azimuth) * FMath::Sin(Elevation) * CosElevation*CosElevation;
+	}
+
+	static const float MagicNumber_ACN_11 = 0.25f * FMath::Sqrt(6.0f);
+	static float SHBasisFunction_ACN_11(float Azimuth, float Elevation)
+	{
+		static const float SinElevation = FMath::Sin(Elevation);
+		return MagicNumber_ACN_11 * FMath::Sin(Azimuth) * FMath::Cos(Elevation) * (5.0f * SinElevation*SinElevation - 1.0f);
+	}
+
+	static float SHBasisFunction_ACN_12(float Azimuth, float Elevation)
+	{
+		const float SinElevation = FMath::Sin(Elevation);
+		return FMath::Sin(Elevation) * (2.5f * SinElevation*SinElevation - 1.5f);
+	}
+
+	static const float MagicNumber_ACN_13 = 0.25f * FMath::Sqrt(6.0f);
+	static float SHBasisFunction_ACN_13(float Azimuth, float Elevation)
+	{
+		const float SinElevation = FMath::Sin(Elevation);
+		return MagicNumber_ACN_13 * FMath::Cos(Azimuth) * FMath::Cos(Elevation)* (5.0f * SinElevation*SinElevation - 1.0f);
+	}
+
+	static const float MagicNumber_ACN_14 = 0.5f * FMath::Sqrt(15.0f);
+	static float SHBasisFunction_ACN_14(float Azimuth, float Elevation)
+	{
+		const float CosElevation = FMath::Cos(Elevation);
+		return MagicNumber_ACN_14 * FMath::Cos(2.0f * Azimuth) * FMath::Sin(Elevation) * CosElevation*CosElevation;
+	}
+
+	static const float MagicNumber_ACN_15 = 0.25f * FMath::Sqrt(10.0f);
+	static float SHBasisFunction_ACN_15(float Azimuth, float Elevation)
+	{
+		static const float CosElevation = FMath::Cos(Elevation);
+		return MagicNumber_ACN_15 * FMath::Cos(3.0f * Azimuth) * CosElevation*CosElevation*CosElevation;
+	}
+#pragma  endregion
+#pragma region FORTH_ORDER
+	static const float MagicNumber_ACN_16 = FMath::Sqrt(35.0f) / 8.0f;
+	static float SHBasisFunction_ACN_16(float Azimuth, float Elevation)
+	{
+		static const float CosElevation = FMath::Cos(Elevation);
+		return MagicNumber_ACN_16 * FMath::Sin(4.0f * Azimuth) * CosElevation*CosElevation*CosElevation*CosElevation;
+	}
+
+	static const float MagicNumber_ACN_17 = 0.25f * FMath::Sqrt(70.0f);
+	static float SHBasisFunction_ACN_17(float Azimuth, float Elevation)
+	{
+		const float CosElevation = FMath::Cos(Elevation);
+		return MagicNumber_ACN_17 * FMath::Sin(3.0f * Azimuth) * FMath::Sin(Elevation) * CosElevation*CosElevation*CosElevation;
+	}
+
+	static const float MagicNumber_ACN_18 = 0.25f * FMath::Sqrt(5.0f);
+	static float SHBasisFunction_ACN_18(float Azimuth, float Elevation)
+	{
+		static const float SinElevation = FMath::Sin(Elevation);
+		static const float CosElevation = FMath::Cos(Elevation);
+		return MagicNumber_ACN_18 * FMath::Sin(2.0f * Azimuth) * (7.0f * SinElevation*SinElevation - 1.0f) * CosElevation*CosElevation;
+	}
+
+	static const float MagicNumber_ACN_19 = FMath::Sqrt(10.0f) / 8.0f;
+	static float SHBasisFunction_ACN_19(float Azimuth, float Elevation)
+	{
+		static const float SinElevation = FMath::Sin(Elevation);
+		return MagicNumber_ACN_19 * FMath::Sin(Azimuth) * FMath::Sin(2.0f * Elevation) * (7.0f * SinElevation*SinElevation - 3.0f);
+	}
+
+	static float SHBasisFunction_ACN_20(float Azimuth, float Elevation)
+	{
+		const float SinElevation = FMath::Sin(Elevation);
+		const float SquaredSinElevation = SinElevation*SinElevation;
+		return 0.625f * SquaredSinElevation*(7.0f * SquaredSinElevation - 3.0f);
+	}
+
+	static const float MagicNumber_ACN_21 = FMath::Sqrt(10.0f) / 8.0f;
+	static float SHBasisFunction_ACN_21(float Azimuth, float Elevation)
+	{
+		static const float SinElevation = FMath::Sin(Elevation);
+		return MagicNumber_ACN_21 * FMath::Cos(Azimuth) * FMath::Sin(2.0f * Elevation) * (7.0f * SinElevation*SinElevation  - 3.0f);
+	}
+	
+	static const float MagicNumber_ACN_22 = 0.25f * FMath::Sqrt(5.0f);
+	static float SHBasisFunction_ACN_22(float Azimuth, float Elevation)
+	{
+		const float SinElevation = FMath::Sin(Elevation);
+		const float CosElevation = FMath::Cos(Elevation);
+		return MagicNumber_ACN_22 * FMath::Cos(2.0f * Azimuth) * (7.0f * SinElevation*SinElevation - 1.0f) * CosElevation*CosElevation;
+	}
+
+	static const float MagicNumber_ACN_23 = 0.25f * FMath::Sqrt(70.0f);
+	static float SHBasisFunction_ACN_23(float Azimuth, float Elevation)
+	{
+		const float CosElevation = FMath::Cos(Elevation);
+		return MagicNumber_ACN_23 * FMath::Cos(3.0f * Azimuth) * FMath::Sin(Elevation) * CosElevation*CosElevation*CosElevation;
+	}
+
+	static const float MagicNumber_ACN_24 = FMath::Sqrt(35.0f) / 8.0f;
+	static float SHBasisFunction_ACN_24(float Azimuth, float Elevation)
+	{
+		const float CosElevation = FMath::Cos(Elevation);
+		return MagicNumber_ACN_24 * FMath::Cos(4.0f * Azimuth) * CosElevation*CosElevation*CosElevation*CosElevation;
+	}
+#pragma endregion
+#pragma region FIFTH_ORDER
+
+	static const float MagicNumber_ACN_25 = 3.0f * FMath::Sqrt(14.0f) / 16.0f;
+	static float SHBasisFunction_ACN_25(const float Azimuth, const float Elevation)
+	{
+		const float CosElevation = FMath::Cos(Elevation);
+		return MagicNumber_ACN_25 * FMath::Sin(5.0f * Azimuth) * CosElevation*CosElevation*CosElevation*CosElevation*CosElevation;
+	}
+
+	static const float MagicNumber_ACN_26 = 3.0f * FMath::Sqrt(35.0f) / 8.0f;
+	static float SHBasisFunction_ACN_26(const float Azimuth, const float Elevation)
+	{
+		const float CosElevation = FMath::Cos(Elevation);
+		return  MagicNumber_ACN_26 * FMath::Sin(4.0f * Azimuth) * FMath::Sin(Elevation) * CosElevation*CosElevation*CosElevation*CosElevation;
+	}
+
+	static const float MagicNumber_ACN_27 = FMath::Sqrt(70.0f) / 32.0f;
+	static float SHBasisFunction_ACN_27(const float Azimuth, const float Elevation)
+	{
+		const float CosElevation = FMath::Cos(Elevation);
+		return MagicNumber_ACN_27 * FMath::Sin(3.0f * Azimuth) * (7.0f - 9.0f * FMath::Cos(2.0f * Elevation)) * CosElevation*CosElevation*CosElevation;
+	}
+
+	static const float MagicNumber_ACN_28 = 0.25f * FMath::Sqrt(105.0f);
+	static float SHBasisFunction_ACN_28(const float Azimuth, const float Elevation)
+	{
+		const float SinElevation = FMath::Sin(Elevation);
+		const float CosElevation = FMath::Cos(Elevation);
+		return MagicNumber_ACN_28 * FMath::Sin(2.0f * Azimuth) * FMath::Sin(Elevation) * (3.0f * SinElevation*SinElevation - 1.0f) * CosElevation*CosElevation;
+	}
+
+	static const float MagicNumber_ACN_29 = FMath::Sqrt(15.0f) / 8.0f;
+	static float SHBasisFunction_ACN_29(const float Azimuth, const float Elevation)
+	{
+		const float SinElevation = FMath::Sin(Elevation);
+		const float SquaredSinElevation = SinElevation * SinElevation;
+		return MagicNumber_ACN_29 * FMath::Sin(Azimuth) * FMath::Cos(Elevation) * (21.0f * SquaredSinElevation*SquaredSinElevation - 14.0f * SquaredSinElevation + 1.0f);
+	}
+
+	static float SHBasisFunction_ACN_30(const float Azimuth, const float Elevation)
+	{
+		const float SinElevation = FMath::Sin(Elevation);
+		const float CubedSinElevation = SinElevation * SinElevation * SinElevation;
+		return (63.0f * CubedSinElevation*SinElevation*SinElevation - 70.0f * CubedSinElevation + 15.0f * FMath::Sin(Elevation)) / 8.0f;
+	}
+
+	static const float MagicNumber_ACN_31 = FMath::Sqrt(15.0f) / 8.0f;
+	static float SHBasisFunction_ACN_31(const float Azimuth, const float Elevation)
+	{
+		const float SinElevation = FMath::Sin(Elevation);
+		const float SquaredSinElevation = SinElevation * SinElevation;
+		return MagicNumber_ACN_31 * FMath::Cos(Azimuth) * FMath::Cos(Elevation) * (21.0f * SquaredSinElevation*SquaredSinElevation - 14.0f * SquaredSinElevation + 1.0f);
+	}
+
+	static const float MagicNumber_ACN_32 = 0.25f * FMath::Sqrt(105.0f);
+	static float SHBasisFunction_ACN_32(const float Azimuth, const float Elevation)
+	{
+		const float SinElevation = FMath::Sin(Elevation);
+		const float CosElevation = FMath::Cos(Elevation);
+		return MagicNumber_ACN_32 * FMath::Cos(2.0f * Azimuth) * FMath::Sin(Elevation) * (3.0f * SinElevation*SinElevation - 1) * CosElevation;
+	}
+
+	static const float MagicNumber_ACN_33 = FMath::Sqrt(70.0f) / 32.0f;
+	static float SHBasisFunction_ACN_33(const float Azimuth, const float Elevation)
+	{
+		const float CosElevation = FMath::Cos(Elevation);
+		return MagicNumber_ACN_33 * FMath::Cos(3.0f * Azimuth) * (7.0f - 9.0f * FMath::Cos(2.0f * Elevation)) * CosElevation*CosElevation*CosElevation;
+	}
+
+	static const float MagicNumber_ACN_34 = 3.0f * FMath::Sqrt(35.0f) / 8.0f;
+	static float SHBasisFunction_ACN_34(const float Azimuth, const float Elevation)
+	{
+		const float CosElevation = FMath::Cos(Elevation);
+		return MagicNumber_ACN_34 * FMath::Cos(4.0f * Azimuth) * FMath::Sin(Elevation) * CosElevation*CosElevation*CosElevation*CosElevation;
+	}
+
+	static const float MagicNumber_ACN_35 = 3.0f * FMath::Sqrt(14.0f) / 16.0f;
+	static float SHBasisFunction_ACN_35(const float Azimuth, const float Elevation)
+	{
+		const float CosElevation = FMath::Cos(Elevation);
+		return MagicNumber_ACN_35 * FMath::Cos(5.0f * Azimuth) * CosElevation*CosElevation*CosElevation*CosElevation*CosElevation;
+	}
+
+#pragma endregion
+} // namespace SHBasis
+
+void FSphericalHarmonicCalculator::ComputeSoundfieldChannelGains(const int32 Order, const float Azimuth, const float Elevation, float* OutGains)
+{
+	check(OutGains);
+	check(Order >= 0 && Order <= 5); // only have basis functions up to 5th order implemented
+
+	switch (Order)
+	{
+		// lack of break; is intentional here.
+		// (i.e. if Order is 2 we need to compute the gains for 1st and 0th order as well)
+		case 5:
+			OutGains[AmbiChanNumber::ACN_35] = SHBasis::SHBasisFunction_ACN_35(Azimuth, Elevation);
+			OutGains[AmbiChanNumber::ACN_34] = SHBasis::SHBasisFunction_ACN_34(Azimuth, Elevation);
+			OutGains[AmbiChanNumber::ACN_33] = SHBasis::SHBasisFunction_ACN_33(Azimuth, Elevation);
+			OutGains[AmbiChanNumber::ACN_32] = SHBasis::SHBasisFunction_ACN_32(Azimuth, Elevation);
+			OutGains[AmbiChanNumber::ACN_31] = SHBasis::SHBasisFunction_ACN_31(Azimuth, Elevation);
+			OutGains[AmbiChanNumber::ACN_30] = SHBasis::SHBasisFunction_ACN_30(Azimuth, Elevation);
+			OutGains[AmbiChanNumber::ACN_29] = SHBasis::SHBasisFunction_ACN_29(Azimuth, Elevation);
+			OutGains[AmbiChanNumber::ACN_28] = SHBasis::SHBasisFunction_ACN_28(Azimuth, Elevation);
+			OutGains[AmbiChanNumber::ACN_27] = SHBasis::SHBasisFunction_ACN_27(Azimuth, Elevation);
+			OutGains[AmbiChanNumber::ACN_26] = SHBasis::SHBasisFunction_ACN_26(Azimuth, Elevation);
+			OutGains[AmbiChanNumber::ACN_25] = SHBasis::SHBasisFunction_ACN_25(Azimuth, Elevation);
+		case 4:
+			OutGains[AmbiChanNumber::ACN_24] = SHBasis::SHBasisFunction_ACN_24(Azimuth, Elevation);
+			OutGains[AmbiChanNumber::ACN_23] = SHBasis::SHBasisFunction_ACN_23(Azimuth, Elevation);
+			OutGains[AmbiChanNumber::ACN_22] = SHBasis::SHBasisFunction_ACN_22(Azimuth, Elevation);
+			OutGains[AmbiChanNumber::ACN_21] = SHBasis::SHBasisFunction_ACN_21(Azimuth, Elevation);
+			OutGains[AmbiChanNumber::ACN_20] = SHBasis::SHBasisFunction_ACN_20(Azimuth, Elevation);
+			OutGains[AmbiChanNumber::ACN_19] = SHBasis::SHBasisFunction_ACN_19(Azimuth, Elevation);
+			OutGains[AmbiChanNumber::ACN_18] = SHBasis::SHBasisFunction_ACN_18(Azimuth, Elevation);
+			OutGains[AmbiChanNumber::ACN_17] = SHBasis::SHBasisFunction_ACN_17(Azimuth, Elevation);
+			OutGains[AmbiChanNumber::ACN_16] = SHBasis::SHBasisFunction_ACN_16(Azimuth, Elevation);
+		case 3:
+			OutGains[AmbiChanNumber::ACN_15] = SHBasis::SHBasisFunction_ACN_15(Azimuth, Elevation);
+			OutGains[AmbiChanNumber::ACN_14] = SHBasis::SHBasisFunction_ACN_14(Azimuth, Elevation);
+			OutGains[AmbiChanNumber::ACN_13] = SHBasis::SHBasisFunction_ACN_13(Azimuth, Elevation);
+			OutGains[AmbiChanNumber::ACN_12] = SHBasis::SHBasisFunction_ACN_12(Azimuth, Elevation);
+			OutGains[AmbiChanNumber::ACN_11] = SHBasis::SHBasisFunction_ACN_11(Azimuth, Elevation);
+			OutGains[AmbiChanNumber::ACN_10] = SHBasis::SHBasisFunction_ACN_10(Azimuth, Elevation);
+			OutGains[AmbiChanNumber::ACN_9] =  SHBasis::SHBasisFunction_ACN_9(Azimuth, Elevation);
+		case 2:
+			OutGains[AmbiChanNumber::ACN_8] = SHBasis::SHBasisFunction_ACN_8(Azimuth, Elevation);
+			OutGains[AmbiChanNumber::ACN_7] = SHBasis::SHBasisFunction_ACN_7(Azimuth, Elevation);
+			OutGains[AmbiChanNumber::ACN_6] = SHBasis::SHBasisFunction_ACN_6(Azimuth, Elevation);
+			OutGains[AmbiChanNumber::ACN_5] = SHBasis::SHBasisFunction_ACN_5(Azimuth, Elevation);
+			OutGains[AmbiChanNumber::ACN_4] = SHBasis::SHBasisFunction_ACN_4(Azimuth, Elevation);
+		case 1:
+			OutGains[AmbiChanNumber::ACN_3] = SHBasis::SHBasisFunction_ACN_3(Azimuth, Elevation);
+			OutGains[AmbiChanNumber::ACN_2] = SHBasis::SHBasisFunction_ACN_2(Azimuth, Elevation);
+			OutGains[AmbiChanNumber::ACN_1] = SHBasis::SHBasisFunction_ACN_1(Azimuth, Elevation);
+		case 0:
+			OutGains[AmbiChanNumber::ACN_0] = SHBasis::SHBasisFunction_ACN_0(Azimuth, Elevation);
+			break;
+	}
+}
+
+void FSphericalHarmonicCalculator::GenerateFirstOrderRotationMatrixGivenRadians(const float RotXRadians, const float RotYRadians, const float RotZRadians, FMatrix & OutMatrix)
+{
+	const float SinX = FMath::Sin(RotXRadians);
+	const float SinY = FMath::Sin(RotYRadians);
+	const float SinZ = FMath::Sin(RotZRadians);
+
+	const float CosX = FMath::Cos(RotXRadians);
+	const float CosY = FMath::Cos(RotYRadians);
+	const float CosZ = FMath::Cos(RotZRadians);
+
+
+	// Build out Rotation Matrix
+	OutMatrix.SetIdentity();
+
+	// row 0:
+	/*	1.0f						0.0f															0.0f									0.0f									*/
+
+	// row 1:
+	/*	0.0f	*/	OutMatrix.M[1][1] = (CosX * CosZ) + (SinX * SinY * SinZ);		OutMatrix.M[1][2] = -SinX * CosY;		OutMatrix.M[1][3] = (CosX * SinZ) - (SinX * SinY * CosZ);
+
+	// row 2:
+	/*	0.0f	*/	OutMatrix.M[2][1] = (SinX * CosZ) - (CosX * SinY * SinZ);		OutMatrix.M[2][2] = CosX * CosY;		OutMatrix.M[2][3] = (CosX * SinY * CosZ) + (SinX * SinZ);
+
+	// row 3:
+	/*	0.0f	*/	OutMatrix.M[3][1] = -CosY * SinZ;								OutMatrix.M[3][2] = -SinY;				OutMatrix.M[3][3] = CosY * CosZ;
+}
+
+void FSphericalHarmonicCalculator::GenerateFirstOrderRotationMatrixGivenDegrees(const float RotXDegrees, const float RotYDegrees, const float RotZDegrees, FMatrix & OutMatrix)
+{
+	constexpr const float DEG_2_RAD = PI / 180.0f;
+	GenerateFirstOrderRotationMatrixGivenRadians(RotXDegrees * DEG_2_RAD, RotYDegrees * DEG_2_RAD, RotZDegrees * DEG_2_RAD, OutMatrix);
+}
diff --git a/Engine/Source/Runtime/SoundFieldRendering/Public/SoundFieldRendering.h b/Engine/Source/Runtime/SoundFieldRendering/Public/SoundFieldRendering.h
new file mode 100644
--- /dev/null
+++ b/Engine/Source/Runtime/SoundFieldRendering/Public/SoundFieldRendering.h
@@ -0,0 +1,106 @@
+// Copyright 1998-2019 Epic Games, Inc. All Rights Reserved.
+
+#pragma once
+
+#include "CoreMinimal.h"
+#include "ISoundfieldFormat.h"
+#include "SphericalHarmonicCalculator.h"
+#include "DSP/BufferVectorOperations.h"
+
+SOUNDFIELDRENDERING_API FName GetUnrealAmbisonicsFormatName();
+SOUNDFIELDRENDERING_API TUniquePtr<ISoundfieldDecoderStream> CreateDefaultSourceAmbisonicsDecoder(Audio::FMixerDevice* InDevice);
+SOUNDFIELDRENDERING_API ISoundfieldEncodingSettingsProxy& GetAmbisonicsSourceDefaultSettings();
+
+
+class SOUNDFIELDRENDERING_API FAmbisonicsSoundfieldBuffer : public ISoundfieldAudioPacket
+{
+public:
+	FAmbisonicsSoundfieldBuffer()
+		: NumChannels(0)
+		, Rotation(FQuat::Identity)
+	{}
+
+	// Interleaved audio buffer for all vector parts of the ambisonics stream.
+	Audio::AlignedFloatBuffer AudioBuffer;
+
+	// number of channels in the Ambisonics stream.
+	// Currently we don't explicitly support mixed order ambisonics,
+	// so this will always be equal to (m + 1)^2, where m is the order  of ambisonics this was encoded with.
+	int32 NumChannels;
+
+	// This is the rotation of the ambisonics source.
+	FQuat Rotation;
+
+	virtual void Serialize(FArchive& Ar) override;
+	virtual TUniquePtr<ISoundfieldAudioPacket> Duplicate() const override;
+	virtual void Reset() override;
+
+};
+
+class SOUNDFIELDRENDERING_API FAmbisonicsSoundfieldSettings : public ISoundfieldEncodingSettingsProxy
+{
+public:
+	int32 Order;
+
+	virtual uint32 GetUniqueId() const override;
+	virtual TUniquePtr<ISoundfieldEncodingSettingsProxy> Duplicate() const override;
+
+};
+
+class SOUNDFIELDRENDERING_API FSoundFieldDecoder
+{
+public:
+	FSoundFieldDecoder(); // initializes VirtualSpeakerWorldLockedSpeakerGains
+
+	void DecodeAudioDirectlyToDeviceOutputPositions(const FAmbisonicsSoundfieldBuffer& InputData, const FSoundfieldSpeakerPositionalData& OutputPositions, Audio::AlignedFloatBuffer& OutputData);
+	void DecodeAudioToSevenOneAndDownmixToDevice(const FAmbisonicsSoundfieldBuffer& InputData, const FSoundfieldSpeakerPositionalData& OutputPositions, Audio::AlignedFloatBuffer& OutputData);
+	static void RotateFirstOrderAmbisonicsBed(const FAmbisonicsSoundfieldBuffer& InputData, FAmbisonicsSoundfieldBuffer& OutputData, const FQuat& DestinationRotation, const FQuat& PreviousRotation);
+
+	static void FoaRotationInPlace(Audio::AlignedFloatBuffer& InOutFrames, const float XRotDegrees, const float YRotDegrees, const float ZRotDegrees);
+
+private:
+	// Special cased, vectorized versions of first order operations.
+	static void FirstOrderDecodeLoop(const int32 NumFrames, const int32 NumOutputChannels, const float* pAmbiFrame, const int32 NumAmbiChannels, float* SpeakerGainsPtr, float* OutputBufferPtrBuffer);
+	static void FirstOrderToSevenOneLoop(const int32 NumFrames, const int32 NumOutputChannels, const float* pAmbiFrame, const int32 NumAmbiChannels, float* SpeakerGainsPtr, float* OutputBufferPtrBuffer);
+
+	// Special cased vectorized versions of 3rd and fifth order operations.
+	static void OddOrderDecodeLoop(const int32 NumFrames, const int32 NumOutputChannels, const float* pAmbiFrame, const int32 NumAmbiChannels, float* SpeakerGainsPtr, float* OutputBufferPtrBuffer);
+	static void OddOrderToSevenOneLoop(const int32 NumFrames, const int32 NumOutputChannels, const float* pAmbiFrame, const int32 NumAmbiChannels, float* SpeakerGainsPtr, float* OutputBufferPtrBuffer);
+
+	// Partially vectorized decode operations for 2nd and 4th order operations.
+	static void EvenOrderDecodeLoop(const int32 NumFrames, const int32 NumOutputChannels, const float* pAmbiFrame, const int32 NumAmbiChannels, float* SpeakerGainsPtr, float* OutputBufferPtrBuffer);
+	static void EvenOrderToSevenOneLoop(const int32 NumFrames, const int32 NumOutputChannels, const float* pAmbiFrame, const int32 NumAmbiChannels, float* SpeakerGainsPtr, float* OutputBufferPtrBuffer);
+
+	// Helper function to get the default channel positions for a channel array.
+	static TArray<Audio::FChannelPositionInfo>* GetDefaultChannelPositions(int32 InNumChannels);
+
+	Audio::AlignedFloatBuffer VirtualSpeakerScratchBuffers;
+	Audio::AlignedFloatBuffer FoaVirtualSpeakerWordLockedGains;
+	Audio::AlignedFloatBuffer TargetSpeakerGains;
+	Audio::AlignedFloatBuffer CurrentSpeakerGains;
+	Audio::AlignedFloatBuffer MixdownGainsMap;
+	Audio::FBufferLinearEase SpeakerGainLerper;
+	FVector2D LastListenerRotationSphericalCoord{ 0.0f, 0.0f };
+
+	static FSphericalHarmonicCalculator SphereHarmCalc;
+	static TArray<Audio::FChannelPositionInfo> VirtualSpeakerLocationsHorzOnly;
+	static const VectorRegister Sqrt2Over2Vec;
+	static const VectorRegister ZeroVec;
+
+	
+	friend class FSoundFieldEncoder;
+};
+
+class SOUNDFIELDRENDERING_API FSoundFieldEncoder
+{
+public:
+	FSoundFieldEncoder();
+	void EncodeAudioDirectlyFromOutputPositions(const Audio::AlignedFloatBuffer& InputData, const FSoundfieldSpeakerPositionalData& InputPositions, const FAmbisonicsSoundfieldSettings& Settings, FAmbisonicsSoundfieldBuffer& OutputData);
+
+private:
+	static void EncodeLoop(const int32 NumFrames, const int32 NumInputChannels, const float* RESTRICT InputAudioPtr, const int32 NumAmbiChannels, float* RESTRICT SpeakerGainsPtr, float* RESTRICT OutputAmbiBuffer);
+
+	Audio::AlignedFloatBuffer VirtualSpeakerScratchBuffers;
+	Audio::AlignedFloatBuffer SpeakerGains;
+	Audio::AlignedFloatBuffer MixdownGainsMap;
+};
\ No newline at end of file
diff --git a/Engine/Source/Runtime/SoundFieldRendering/Public/SphericalHarmonicCalculator.h b/Engine/Source/Runtime/SoundFieldRendering/Public/SphericalHarmonicCalculator.h
new file mode 100644
--- /dev/null
+++ b/Engine/Source/Runtime/SoundFieldRendering/Public/SphericalHarmonicCalculator.h
@@ -0,0 +1,38 @@
+// Copyright 1998-2019 Epic Games, Inc. All Rights Reserved.
+
+#pragma once
+
+#include "CoreMinimal.h"
+#include "DSP/BufferVectorOperations.h"
+
+
+/*
+	Azimuth angle is measured CCW from front.
+	Elevation is 0 horizontal plane, + is above horizontal plane
+*/
+
+class SOUNDFIELDRENDERING_API FSphericalHarmonicCalculator
+{
+public:
+	enum AmbiChanNumber
+	{
+		/* 0th-Order */	ACN_0 = 0,
+		/* 1st-Order */	ACN_1, ACN_2, ACN_3,
+		/* 2nd-Order */	ACN_4, ACN_5, ACN_6, ACN_7, ACN_8,
+		/* 3rd-Order */	ACN_9, ACN_10, ACN_11, ACN_12, ACN_13, ACN_14, ACN_15,
+		/* 4th-Order */	ACN_16, ACN_17, ACN_18, ACN_19, ACN_20, ACN_21, ACN_22, ACN_23, ACN_24,
+		/* 5th-Order */	ACN_25, ACN_26, ACN_27, ACN_28, ACN_29, ACN_30, ACN_31, ACN_32, ACN_33, ACN_34, ACN_35
+	};
+
+	static void ComputeSoundfieldChannelGains(const int32 Order, const float Azimuth, const float Elevation, float* OutGains);
+
+	static void GenerateFirstOrderRotationMatrixGivenRadians(const float RotXRadians, const float RotYRadians, const float RotZRadians, FMatrix& OutMatrix);
+	static void GenerateFirstOrderRotationMatrixGivenDegrees(const float RotXDegrees, const float RotYDegrees, const float RotZDegrees, FMatrix& OutMatrix);
+
+	static void AdjustUESphericalCoordinatesForAmbisonics(FVector2D& InOutVector)
+	{
+		InOutVector.X = -(InOutVector.X - HALF_PI);
+		InOutVector.Y *= -1.0f;
+		FMemory::Memswap(&InOutVector.X,&InOutVector.Y, sizeof(InOutVector.X));
+	}
+};
