Fix for -Waddress-of-temporary.

diff --git a/Engine/Source/Runtime/SoundFieldRendering/Private/SoundFieldRendering.cpp b/Engine/Source/Runtime/SoundFieldRendering/Private/SoundFieldRendering.cpp
--- a/Engine/Source/Runtime/SoundFieldRendering/Private/SoundFieldRendering.cpp
+++ b/Engine/Source/Runtime/SoundFieldRendering/Private/SoundFieldRendering.cpp
@@ -397,7 +397,8 @@ void FSoundFieldDecoder::DecodeAudioToSevenOneAndDownmixToDevice(const FAmbisoni
 			OutputBufferPtr[OutSample] = VectorGetComponent(VectorDot4(CurrAmbiFrame, CurrSpeakerGain), 0);
 
 			// step forward gain interpolation
-			VectorStoreAligned(VectorAdd(CurrSpeakerGain, CurrSpeakerGainDelta), SpeakerGainsPtr + SpeakerGainOffset);
+			VectorRegister UpdatedSpeakerGain = VectorAdd(CurrSpeakerGain, CurrSpeakerGainDelta);
+			VectorStoreAligned(UpdatedSpeakerGain, SpeakerGainsPtr + SpeakerGainOffset);
 		}
 	}
 
