Fixes for various static analysis warnings.

diff --git a/Engine/Source/Editor/AudioEditor/Private/SoundSubmixGraphSchema.cpp b/Engine/Source/Editor/AudioEditor/Private/SoundSubmixGraphSchema.cpp
--- a/Engine/Source/Editor/AudioEditor/Private/SoundSubmixGraphSchema.cpp
+++ b/Engine/Source/Editor/AudioEditor/Private/SoundSubmixGraphSchema.cpp
@@ -50,8 +50,10 @@ FSoundSubmixGraphConnectionDrawingPolicy::FSoundSubmixGraphConnectionDrawingPoli
 // Give specific editor modes a chance to highlight this connection or darken non-interesting connections
 void FSoundSubmixGraphConnectionDrawingPolicy::DetermineWiringStyle(UEdGraphPin* OutputPin, UEdGraphPin* InputPin, FConnectionParams& OutParams)
 {
-	check(GraphObj);
-	check(OutputPin);
+	if (!ensure(OutputPin && InputPin && GraphObj))
+	{
+		return;
+	}
 
 	OutParams.AssociatedPin1 = InputPin;
 	OutParams.AssociatedPin2 = OutputPin;
diff --git a/Engine/Source/Runtime/Engine/Private/SoundWave.cpp b/Engine/Source/Runtime/Engine/Private/SoundWave.cpp
--- a/Engine/Source/Runtime/Engine/Private/SoundWave.cpp
+++ b/Engine/Source/Runtime/Engine/Private/SoundWave.cpp
@@ -532,7 +532,10 @@ bool USoundWave::HasCompressedData(FName Format, ITargetPlatform* TargetPlatform
 
 	if (GIsEditor)
 	{
-		CompressionOverrides = (TargetPlatform) ? FPlatformCompressionUtilities::GetCookOverrides(*TargetPlatform->IniPlatformName()) : nullptr;
+		if (TargetPlatform)
+		{
+			CompressionOverrides = FPlatformCompressionUtilities::GetCookOverrides(*TargetPlatform->IniPlatformName());
+		}
 	}
 	else
 	{
@@ -1493,7 +1496,7 @@ void USoundWave::PostEditChangeProperty(FPropertyChangedEvent& PropertyChangedEv
 	if (PropertyChangedEvent.ChangeType != EPropertyChangeType::Interactive)
 	{
 		// Regenerate on save any compressed sound formats or if analysis needs to be re-done
-		if (FProperty* PropertyThatChanged = PropertyChangedEvent.Property)
+		if (UProperty* PropertyThatChanged = PropertyChangedEvent.Property)
 		{
 			const FName& Name = PropertyThatChanged->GetFName();
 			if (Name == CompressionQualityFName || Name == StreamingFName || Name == SeekableStreamingFName)
diff --git a/Engine/Source/Runtime/SoundFieldRendering/Private/SoundFieldRendering.cpp b/Engine/Source/Runtime/SoundFieldRendering/Private/SoundFieldRendering.cpp
--- a/Engine/Source/Runtime/SoundFieldRendering/Private/SoundFieldRendering.cpp
+++ b/Engine/Source/Runtime/SoundFieldRendering/Private/SoundFieldRendering.cpp
@@ -361,9 +361,9 @@ void FSoundFieldDecoder::DecodeAudioToSevenOneAndDownmixToDevice(const FAmbisoni
 	if (!bShouldInterpolate) // Non-Interpolated decode
 	{
 		// for each output frame...
-		for (int32 OutSample = 0; OutSample < NumVirtualOutputSamples; ++OutSample)
+		for (uint64 OutSample = 0; OutSample < NumVirtualOutputSamples; ++OutSample)
 		{
-			const int32 ChannelOffset = OutSample & OutputChannelIndexWrapMask;
+			const uint64 ChannelOffset = OutSample & OutputChannelIndexWrapMask;
 
 			// (non bit-shift versions:)
 			// CurrAmbiFrame = VectorLoadAligned(AmbiInputBufferPtr + OutSample / 8 * NumAmbiChannels);
@@ -382,10 +382,10 @@ void FSoundFieldDecoder::DecodeAudioToSevenOneAndDownmixToDevice(const FAmbisoni
 		VectorRegister CurrSpeakerGainDelta = VectorRegister();
 
 		// for each output frame...
-		for (int32 OutSample = 0; OutSample < NumVirtualOutputSamples; ++OutSample)
+		for (uint64 OutSample = 0; OutSample < NumVirtualOutputSamples; ++OutSample)
 		{
-			const int32 ChannelOffset = OutSample & OutputChannelIndexWrapMask;
-			const int32 SpeakerGainOffset = ChannelOffset << 2;
+			const uint64 ChannelOffset = OutSample & OutputChannelIndexWrapMask;
+			const uint64 SpeakerGainOffset = ChannelOffset << 2;
 
 			// (non bit-shift versions:)
 			// CurrAmbiFrame = VectorLoadAligned(AmbiInputBufferPtr + OutSample / 8 * NumAmbiChannels);
@@ -548,9 +548,9 @@ void FSoundFieldDecoder::FirstOrderToSevenOneLoop(const int32 NumFrames, const i
 	constexpr const int32 AmbiChannelIndexWrapMask = ~0x3; // 8 virtual output channels, so to mod we clear out all but the 2 LS bits 111111100b
 
 	// for each output frame...
-	for (int32 OutSample = 0; OutSample < NumVirtualOutputSamples; ++OutSample)
+	for (uint64 OutSample = 0; OutSample < NumVirtualOutputSamples; ++OutSample)
 	{
-		const int32 ChannelOffset = OutSample & OutputChannelIndexWrapMask;
+		const uint64 ChannelOffset = OutSample & OutputChannelIndexWrapMask;
 
 		// (non bit-shift versions:)
 		// CurrAmbiFrame = VectorLoadAligned(AmbiInputBufferPtr + OutSample / 8 * NumAmbiChannels);
