CIS Linux fix (removing #pragma once from .cpp)

diff --git a/Engine/Source/Runtime/SoundFieldRendering/Private/SphericalHarmonicCalculator.cpp b/Engine/Source/Runtime/SoundFieldRendering/Private/SphericalHarmonicCalculator.cpp
--- a/Engine/Source/Runtime/SoundFieldRendering/Private/SphericalHarmonicCalculator.cpp
+++ b/Engine/Source/Runtime/SoundFieldRendering/Private/SphericalHarmonicCalculator.cpp
@@ -1,6 +1,5 @@
 // Copyright Epic Games, Inc. All Rights Reserved.
 
-#pragma once
 #include "SphericalHarmonicCalculator.h"
 #include "DSP/Dsp.h"
 
