Fixed issue with ambisonic decoding where rotation had to occor in both elevation AND azimuth to update the decoding (instead of either/or)

diff --git a/Engine/Source/Runtime/SoundFieldRendering/Private/SoundFieldRendering.cpp b/Engine/Source/Runtime/SoundFieldRendering/Private/SoundFieldRendering.cpp
--- a/Engine/Source/Runtime/SoundFieldRendering/Private/SoundFieldRendering.cpp
+++ b/Engine/Source/Runtime/SoundFieldRendering/Private/SoundFieldRendering.cpp
@@ -146,8 +146,8 @@ void FSoundFieldDecoder::DecodeAudioDirectlyToDeviceOutputPositions(const FAmbis
 
 	// Don't interpolate if the number of speakers has changed, or if listener rotation has not changed
 	bool bShouldInterpolate = TargetSpeakerGains.Num() == CurrentSpeakerGains.Num();
-	bShouldInterpolate &= !FMath::IsNearlyEqual(ListenerRotationSphericalCoord.X, LastListenerRotationSphericalCoord.X);
-	bShouldInterpolate &= !FMath::IsNearlyEqual(ListenerRotationSphericalCoord.Y, LastListenerRotationSphericalCoord.Y);
+	bShouldInterpolate |= !FMath::IsNearlyEqual(ListenerRotationSphericalCoord.X, LastListenerRotationSphericalCoord.X);
+	bShouldInterpolate |= !FMath::IsNearlyEqual(ListenerRotationSphericalCoord.Y, LastListenerRotationSphericalCoord.Y);
 
 	// fill out the ambisonic speaker gain maps
 	for (int32 OutChan = 0; OutChan < NumOutputChannels; ++OutChan)
@@ -339,7 +339,7 @@ void FSoundFieldDecoder::DecodeAudioToSevenOneAndDownmixToDevice(const FAmbisoni
 
 	// Don't interpolate if the listener rotation has not changed
 	bool bShouldInterpolate = !FMath::IsNearlyEqual(ListenerRotationSphericalCoord.X, LastListenerRotationSphericalCoord.X);
-	bShouldInterpolate &= !FMath::IsNearlyEqual(ListenerRotationSphericalCoord.Y, LastListenerRotationSphericalCoord.Y);
+	bShouldInterpolate |= !FMath::IsNearlyEqual(ListenerRotationSphericalCoord.Y, LastListenerRotationSphericalCoord.Y);
 
 	if (bShouldInterpolate)
 	{
