Fixup missing vtable symbol.

- Potentially related to UE-108352.

#rnx

diff --git a/Engine/Source/Runtime/SoundFieldRendering/Public/SoundFieldRendering.h b/Engine/Source/Runtime/SoundFieldRendering/Public/SoundFieldRendering.h
--- a/Engine/Source/Runtime/SoundFieldRendering/Public/SoundFieldRendering.h
+++ b/Engine/Source/Runtime/SoundFieldRendering/Public/SoundFieldRendering.h
@@ -33,6 +33,9 @@ public:
 		, PreviousRotation(FQuat::Identity)
 	{}
 
+	virtual ~FAmbisonicsSoundfieldBuffer()
+	{}
+
 	virtual void Serialize(FArchive& Ar) override;
 	virtual TUniquePtr<ISoundfieldAudioPacket> Duplicate() const override;
 	virtual void Reset() override;
@@ -105,4 +108,4 @@ private:
 	Audio::AlignedFloatBuffer VirtualSpeakerScratchBuffers;
 	Audio::AlignedFloatBuffer SpeakerGains;
 	Audio::AlignedFloatBuffer MixdownGainsMap;
-};
\ No newline at end of file
+};
